Added Live Link OBS Studio Device

#Synced-CL 52582695

diff --git a/Engine/Plugins/VirtualProduction/LiveLinkOBSDevice/LiveLinkOBSDevice.uplugin b/Engine/Plugins/VirtualProduction/LiveLinkOBSDevice/LiveLinkOBSDevice.uplugin
new file mode 100644
--- /dev/null
+++ b/Engine/Plugins/VirtualProduction/LiveLinkOBSDevice/LiveLinkOBSDevice.uplugin
@@ -0,0 +1,31 @@
+{
+    "FileVersion": 3,
+    "Version": 1,
+    "VersionName": "1.0",
+    "FriendlyName": "Live Link Hub OBS Device",
+    "Description": "Live Link Hub device for controlling OBS Studio recording via the OBS WebSocket v5 protocol.",
+    "Category": "Virtual Production",
+    "CreatedBy": "Epic Games, Inc.",
+    "CreatedByURL": "https://epicgames.com",
+    "CanContainContent": false,
+    "IsBetaVersion": false,
+    "IsExperimentalVersion": true,
+    "Installed": false,
+    "Modules": [
+        {
+            "Name": "LiveLinkOBSDevice",
+            "Type": "Editor",
+            "LoadingPhase": "Default"
+        }
+    ],
+    "Plugins": [
+        {
+            "Name": "LiveLinkDevice",
+            "Enabled": true
+        },
+        {
+            "Name": "LiveLinkHub",
+            "Enabled": true
+        }
+    ]
+}
diff --git a/Engine/Plugins/VirtualProduction/LiveLinkOBSDevice/Source/LiveLinkOBSDevice/LiveLinkOBSDevice.Build.cs b/Engine/Plugins/VirtualProduction/LiveLinkOBSDevice/Source/LiveLinkOBSDevice/LiveLinkOBSDevice.Build.cs
new file mode 100644
--- /dev/null
+++ b/Engine/Plugins/VirtualProduction/LiveLinkOBSDevice/Source/LiveLinkOBSDevice/LiveLinkOBSDevice.Build.cs
@@ -0,0 +1,42 @@
+// Copyright Epic Games, Inc. All Rights Reserved.
+
+using UnrealBuildTool;
+
+public class LiveLinkOBSDevice : ModuleRules
+{
+    public LiveLinkOBSDevice(ReadOnlyTargetRules Target) : base(Target)
+    {
+        PrivateDependencyModuleNames.AddRange(
+            new string[]
+            {
+                "Core",
+                "CoreUObject",
+                "Engine",
+                "Json",
+                "LiveLinkDevice",
+                "LiveLinkHub",
+                "LiveLinkInterface",
+                "NamingTokens",
+                "NamingTokensUI",
+                "PropertyEditor",
+                "Slate",
+                "SlateCore",
+                "WebSockets",
+            }
+        );
+
+        if (Target.Platform == UnrealTargetPlatform.Win64
+         || Target.Platform == UnrealTargetPlatform.Mac
+         || Target.Platform == UnrealTargetPlatform.Linux)
+        {
+            AddEngineThirdPartyPrivateStaticDependencies(Target, "OpenSSL");
+            PublicDefinitions.Add("LLOBS_WITH_OPENSSL=1");
+        }
+        else
+        {
+            PublicDefinitions.Add("LLOBS_WITH_OPENSSL=0");
+        }
+
+        CppCompileWarningSettings.UnsafeTypeCastWarningLevel = WarningLevel.Error;
+    }
+}
diff --git a/Engine/Plugins/VirtualProduction/LiveLinkOBSDevice/Source/LiveLinkOBSDevice/Private/Customization/LiveLinkOBSDeviceSettingsCustomization.cpp b/Engine/Plugins/VirtualProduction/LiveLinkOBSDevice/Source/LiveLinkOBSDevice/Private/Customization/LiveLinkOBSDeviceSettingsCustomization.cpp
new file mode 100644
--- /dev/null
+++ b/Engine/Plugins/VirtualProduction/LiveLinkOBSDevice/Source/LiveLinkOBSDevice/Private/Customization/LiveLinkOBSDeviceSettingsCustomization.cpp
@@ -0,0 +1,79 @@
+// Copyright Epic Games, Inc. All Rights Reserved.
+
+#include "LiveLinkOBSDeviceSettingsCustomization.h"
+
+#include "DetailLayoutBuilder.h"
+#include "DetailWidgetRow.h"
+#include "Devices/LiveLinkOBSDevice.h"
+#include "IDetailPropertyRow.h"
+#include "NamingTokensSpecifiers.h"
+#include "SNamingTokensEditableTextBox.h"
+
+#define LOCTEXT_NAMESPACE "LiveLinkOBSDevice"
+
+
+TSharedRef<IDetailCustomization> FLiveLinkOBSDeviceSettingsCustomization::MakeInstance()
+{
+    return MakeShared<FLiveLinkOBSDeviceSettingsCustomization>();
+}
+
+void FLiveLinkOBSDeviceSettingsCustomization::CustomizeDetails(IDetailLayoutBuilder& DetailBuilder)
+{
+    IDetailCategoryBuilder& Category = DetailBuilder.EditCategory("OBS Device");
+
+    TArray<TSharedRef<IPropertyHandle>> DefaultProperties;
+    Category.GetDefaultProperties(DefaultProperties);
+
+    // Namespace filter: restrict autocomplete to LLH tokens only.
+    FNamingTokenFilterArgs FilterArgs;
+    FilterArgs.bIncludeGlobal = false;
+    FilterArgs.AdditionalNamespacesToInclude.Add(TEXT("llh"));
+
+    // Hide all properties first. Without this, any explicit AddProperty call places
+    // the row before the auto-generated block, putting FilenameFormat above DisplayName.
+    for (const TSharedRef<IPropertyHandle>& PropertyHandle : DefaultProperties)
+    {
+        DetailBuilder.HideProperty(PropertyHandle);
+    }
+
+    // Re-add all properties in their original declaration order.
+    // FilenameFormat (tagged meta=(NamingTokens)) gets the custom autocomplete widget.
+    // All other properties are re-added with their default widgets unchanged.
+    for (TSharedRef<IPropertyHandle>& PropertyHandle : DefaultProperties)
+    {
+        IDetailPropertyRow& Row = Category.AddProperty(PropertyHandle);
+
+        if (!PropertyHandle->HasMetaData(*UE::NamingTokens::Specifiers::UseNamingTokens))
+        {
+            continue;
+        }
+
+        Row.CustomWidget()
+            .NameContent()
+            [
+                PropertyHandle->CreatePropertyNameWidget(
+                    LOCTEXT("FilenameFormatLabel", "Output Filename Format"))
+            ]
+            .ValueContent()
+            [
+                SNew(SNamingTokensEditableTextBox)
+                // Show the raw token string; the user is editing a format template,
+                // not viewing a resolved filename.
+                .ShouldEvaluateTokens(false)
+                .FilterArgs(FilterArgs)
+                .NamespaceSuggestionPriority(TArray<FString>{ TEXT("llh") })
+                .Text_Lambda([PropertyHandle]() -> FText
+                {
+                    FString Value;
+                    PropertyHandle->GetValue(Value);
+                    return FText::FromString(Value);
+                })
+                .OnTextChanged_Lambda([PropertyHandle](const FText& InText)
+                {
+                    PropertyHandle->SetValue(InText.ToString());
+                })
+            ];
+    }
+}
+
+#undef LOCTEXT_NAMESPACE
diff --git a/Engine/Plugins/VirtualProduction/LiveLinkOBSDevice/Source/LiveLinkOBSDevice/Private/Customization/LiveLinkOBSDeviceSettingsCustomization.h b/Engine/Plugins/VirtualProduction/LiveLinkOBSDevice/Source/LiveLinkOBSDevice/Private/Customization/LiveLinkOBSDeviceSettingsCustomization.h
new file mode 100644
--- /dev/null
+++ b/Engine/Plugins/VirtualProduction/LiveLinkOBSDevice/Source/LiveLinkOBSDevice/Private/Customization/LiveLinkOBSDeviceSettingsCustomization.h
@@ -0,0 +1,18 @@
+// Copyright Epic Games, Inc. All Rights Reserved.
+
+#pragma once
+
+#include "IDetailCustomization.h"
+
+class IDetailLayoutBuilder;
+
+
+class FLiveLinkOBSDeviceSettingsCustomization : public IDetailCustomization
+{
+public:
+    static TSharedRef<IDetailCustomization> MakeInstance();
+
+    //~ Begin IDetailCustomization interface
+    virtual void CustomizeDetails(IDetailLayoutBuilder& DetailBuilder) override;
+    //~ End IDetailCustomization interface
+};
diff --git a/Engine/Plugins/VirtualProduction/LiveLinkOBSDevice/Source/LiveLinkOBSDevice/Private/Devices/LiveLinkOBSDevice.cpp b/Engine/Plugins/VirtualProduction/LiveLinkOBSDevice/Source/LiveLinkOBSDevice/Private/Devices/LiveLinkOBSDevice.cpp
new file mode 100644
--- /dev/null
+++ b/Engine/Plugins/VirtualProduction/LiveLinkOBSDevice/Source/LiveLinkOBSDevice/Private/Devices/LiveLinkOBSDevice.cpp
@@ -0,0 +1,644 @@
+// Copyright Epic Games, Inc. All Rights Reserved.
+
+#include "LiveLinkOBSDevice.h"
+
+#include "Dom/JsonObject.h"
+#include "Dom/JsonValue.h"
+#include "Engine/Engine.h"
+#include "IWebSocket.h"
+
+#if LLOBS_WITH_OPENSSL
+#include <openssl/sha.h>
+#endif
+#include "LiveLinkHub/ILiveLinkRecordingSession.h"
+#include "Logging/StructuredLog.h"
+#include "Misc/Base64.h"
+#include "NamingTokensEngineSubsystem.h"
+#include "Serialization/JsonReader.h"
+#include "Serialization/JsonSerializer.h"
+#include "Serialization/JsonWriter.h"
+#include "WebSocketsModule.h"
+
+
+DEFINE_LOG_CATEGORY(LogLiveLinkOBSDevice);
+
+#define LOCTEXT_NAMESPACE "LiveLinkOBSDevice"
+
+
+// -----------------------------------------------------------------------
+// OBS WebSocket v5 op-codes
+// [REDACTED_PROFILE_URL]/obs-websocket/blob/master/docs/generated/protocol.md
+// -----------------------------------------------------------------------
+namespace OBSWebSocketOpCode
+{
+    static constexpr int32 Hello      = 0;  // OBS → Client
+    static constexpr int32 Identify   = 1;  // Client → OBS
+    static constexpr int32 Identified = 2;  // OBS → Client
+    static constexpr int32 Event      = 5;  // OBS → Client
+    static constexpr int32 Request    = 6;  // Client → OBS
+    static constexpr int32 Response   = 7;  // OBS → Client
+}
+
+
+// -----------------------------------------------------------------------
+// ULiveLinkDevice interface
+// -----------------------------------------------------------------------
+
+FText ULiveLinkOBSDevice::GetDisplayName() const
+{
+    return FText::FromString(GetSettings()->DisplayName);
+}
+
+EDeviceHealth ULiveLinkOBSDevice::GetDeviceHealth() const
+{
+    if (bAuthFailed)
+    {
+        return EDeviceHealth::Error;
+    }
+    if (bPasswordRequired)
+    {
+        return EDeviceHealth::Warning;
+    }
+
+    switch (ConnectionStatus)
+    {
+        case ELiveLinkDeviceConnectionStatus::Connected:    return bIdentified ? EDeviceHealth::Good : EDeviceHealth::Info;
+        case ELiveLinkDeviceConnectionStatus::Connecting:    return EDeviceHealth::Warning;
+        case ELiveLinkDeviceConnectionStatus::Disconnecting: return EDeviceHealth::Warning;
+        case ELiveLinkDeviceConnectionStatus::Disconnected: return EDeviceHealth::Warning;
+        default:                                             return EDeviceHealth::Error;
+    }
+}
+
+FText ULiveLinkOBSDevice::GetHealthText() const
+{
+    if (bAuthFailed)
+    {
+        return LOCTEXT("HealthText_AuthFailed", "Authentication failed. Check password.");
+    }
+    if (bPasswordRequired)
+    {
+        return LOCTEXT("HealthText_PasswordRequired", "OBS requires a password. Enter one in device settings.");
+    }
+
+    switch (ConnectionStatus)
+    {
+        case ELiveLinkDeviceConnectionStatus::Connected:
+            return bIdentified
+                ? FText::GetEmpty()
+                : LOCTEXT("HealthText_Authenticating", "Authenticating...");
+        case ELiveLinkDeviceConnectionStatus::Connecting:
+            return LOCTEXT("HealthText_Connecting", "Connecting...");
+        case ELiveLinkDeviceConnectionStatus::Disconnecting:
+            return LOCTEXT("HealthText_Disconnecting", "Disconnecting...");
+        default:
+            return LOCTEXT("HealthText_Disconnected", "Not connected");
+    }
+}
+
+void ULiveLinkOBSDevice::OnDeviceAdded()
+{
+    if (IModularFeatures::Get().IsModularFeatureAvailable(ILiveLinkRecordingSession::GetModularFeatureName()))
+    {
+        ILiveLinkRecordingSession& Session = ILiveLinkRecordingSession::Get();
+        Session.OnSlateNameChanged().AddUObject(this, &ULiveLinkOBSDevice::HandleSlateNameChanged);
+        Session.OnTakeNumberChanged().AddUObject(this, &ULiveLinkOBSDevice::HandleTakeNumberChanged);
+    }
+    Connect();
+}
+
+void ULiveLinkOBSDevice::OnDeviceRemoved()
+{
+    if (IModularFeatures::Get().IsModularFeatureAvailable(ILiveLinkRecordingSession::GetModularFeatureName()))
+    {
+        ILiveLinkRecordingSession& Session = ILiveLinkRecordingSession::Get();
+        Session.OnSlateNameChanged().RemoveAll(this);
+        Session.OnTakeNumberChanged().RemoveAll(this);
+    }
+    Disconnect();
+}
+
+void ULiveLinkOBSDevice::OnSettingChanged(const FPropertyChangedEvent& InPropertyChangedEvent)
+{
+    Super::OnSettingChanged(InPropertyChangedEvent);
+
+    const FName Changed = InPropertyChangedEvent.GetPropertyName();
+
+    static const FName HostName     = GET_MEMBER_NAME_CHECKED(ULiveLinkOBSDeviceSettings, Host);
+    static const FName PortName     = GET_MEMBER_NAME_CHECKED(ULiveLinkOBSDeviceSettings, Port);
+    static const FName PasswordName = GET_MEMBER_NAME_CHECKED(ULiveLinkOBSDeviceSettings, Password);
+    static const FName FilenameFormatName = GET_MEMBER_NAME_CHECKED(ULiveLinkOBSDeviceSettings, FilenameFormat);
+
+    if (Changed == HostName || Changed == PortName || Changed == PasswordName)
+    {
+        bReconnectAfterSettingsChange = true;
+        Disconnect();
+    }
+    else if (Changed == FilenameFormatName)
+    {
+        SendFilenameToOBS();
+    }
+}
+
+
+// -----------------------------------------------------------------------
+// ILiveLinkDeviceCapability_Connection
+// -----------------------------------------------------------------------
+
+ELiveLinkDeviceConnectionStatus ULiveLinkOBSDevice::GetConnectionStatus_Implementation() const
+{
+    return ConnectionStatus;
+}
+
+FString ULiveLinkOBSDevice::GetHardwareId_Implementation() const
+{
+    const ULiveLinkOBSDeviceSettings* S = GetSettings();
+    return FString::Printf(TEXT("%s:%d"), *S->Host, S->Port);
+}
+
+bool ULiveLinkOBSDevice::Connect_Implementation()
+{
+    if (ConnectionStatus == ELiveLinkDeviceConnectionStatus::Connected ||
+        ConnectionStatus == ELiveLinkDeviceConnectionStatus::Connecting)
+    {
+        return false;
+    }
+    Connect();
+    return true;
+}
+
+bool ULiveLinkOBSDevice::Disconnect_Implementation()
+{
+    if (ConnectionStatus == ELiveLinkDeviceConnectionStatus::Disconnected)
+    {
+        return false;
+    }
+    Disconnect();
+    return true;
+}
+
+void ULiveLinkOBSDevice::SetConnectionStatus(ELiveLinkDeviceConnectionStatus InStatus)
+{
+    ConnectionStatus = InStatus;
+    ILiveLinkDeviceCapability_Connection::SetConnectionStatus(InStatus);
+}
+
+
+// -----------------------------------------------------------------------
+// ILiveLinkDeviceCapability_Recording
+// -----------------------------------------------------------------------
+
+bool ULiveLinkOBSDevice::StartRecording_Implementation()
+{
+    if (bIsRecording || !bIdentified)
+    {
+        return false;
+    }
+
+    SendFilenameToOBS();
+
+    if (SendRequest(TEXT("StartRecord")).IsEmpty())
+    {
+        return false;
+    }
+
+    UE_LOGFMT(LogLiveLinkOBSDevice, Log, "{DeviceName}: Sending StartRecord.", GetSettings()->DisplayName);
+
+    return true;
+}
+
+bool ULiveLinkOBSDevice::StopRecording_Implementation()
+{
+    if (!bIsRecording || !bIdentified)
+    {
+        return false;
+    }
+
+    if (SendRequest(TEXT("StopRecord")).IsEmpty())
+    {
+        return false;
+    }
+
+    UE_LOGFMT(LogLiveLinkOBSDevice, Log, "{DeviceName}: Sending StopRecord.", GetSettings()->DisplayName);
+
+    return true;
+}
+
+bool ULiveLinkOBSDevice::IsRecording_Implementation() const
+{
+    return bIsRecording;
+}
+
+
+// -----------------------------------------------------------------------
+// WebSocket lifecycle
+// -----------------------------------------------------------------------
+
+void ULiveLinkOBSDevice::Connect()
+{
+    if (ConnectionStatus != ELiveLinkDeviceConnectionStatus::Disconnected)
+    {
+        return;
+    }
+
+    bAuthFailed = false;
+    bPasswordRequired = false;
+
+    const ULiveLinkOBSDeviceSettings* S = GetSettings();
+    const FString Url = FString::Printf(TEXT("ws://%s:%d"), *S->Host, S->Port);
+
+    UE_LOGFMT(LogLiveLinkOBSDevice, Log, "{DeviceName}: Connecting to {Url}.", S->DisplayName, Url);
+
+    WebSocket = FWebSocketsModule::Get().CreateWebSocket(Url, TEXT("obswebsocket.v5"));
+    WebSocket->OnConnected()      .AddUObject(this, &ULiveLinkOBSDevice::HandleConnected);
+    WebSocket->OnMessage()        .AddUObject(this, &ULiveLinkOBSDevice::HandleMessage);
+    WebSocket->OnClosed()         .AddUObject(this, &ULiveLinkOBSDevice::HandleClosed);
+    WebSocket->OnConnectionError().AddUObject(this, &ULiveLinkOBSDevice::HandleError);
+
+    SetConnectionStatus(ELiveLinkDeviceConnectionStatus::Connecting);
+    WebSocket->Connect();
+}
+
+void ULiveLinkOBSDevice::Disconnect()
+{
+    bIntentionalDisconnect = true;
+
+    if (WebSocket && WebSocket->IsConnected())
+    {
+        SetConnectionStatus(ELiveLinkDeviceConnectionStatus::Disconnecting);
+        WebSocket->Close();
+    }
+    else
+    {
+        OnDisconnect();
+    }
+}
+
+void ULiveLinkOBSDevice::OnDisconnect()
+{
+    bIdentified    = false;
+    bIsRecording   = false;
+    RequestCounter = 0;
+
+    if (WebSocket)
+    {
+        WebSocket->OnConnected()      .RemoveAll(this);
+        WebSocket->OnMessage()        .RemoveAll(this);
+        WebSocket->OnClosed()         .RemoveAll(this);
+        WebSocket->OnConnectionError().RemoveAll(this);
+        WebSocket.Reset();
+    }
+
+    SetConnectionStatus(ELiveLinkDeviceConnectionStatus::Disconnected);
+
+    // Cancel any previously scheduled reconnect.
+    if (ReconnectTickerHandle.IsValid())
+    {
+        FTSTicker::GetCoreTicker().RemoveTicker(ReconnectTickerHandle);
+        ReconnectTickerHandle.Reset();
+    }
+
+    if (!bIntentionalDisconnect)
+    {
+        UE_LOGFMT(LogLiveLinkOBSDevice, Log, "{DeviceName}: Unexpected disconnect, reconnecting in {Delay}s.",
+            GetSettings()->DisplayName, ReconnectDelay);
+
+        ReconnectTickerHandle = FTSTicker::GetCoreTicker().AddTicker(
+            FTickerDelegate::CreateWeakLambda(this, [this](float) -> bool
+            {
+                ReconnectTickerHandle.Reset();
+                Connect();
+                return false; // one-shot
+            }),
+            ReconnectDelay
+        );
+    }
+
+    bIntentionalDisconnect = false;
+
+    if (bReconnectAfterSettingsChange)
+    {
+        bReconnectAfterSettingsChange = false;
+        Connect();
+    }
+}
+
+
+// -----------------------------------------------------------------------
+// WebSocket event handlers
+// -----------------------------------------------------------------------
+
+void ULiveLinkOBSDevice::HandleConnected()
+{
+    // OBS sends Hello (op 0) immediately after connection; wait for it before doing anything.
+    UE_LOGFMT(LogLiveLinkOBSDevice, Log, "{DeviceName}: WebSocket connected; waiting for Hello.", GetSettings()->DisplayName);
+    SetConnectionStatus(ELiveLinkDeviceConnectionStatus::Connected);
+}
+
+void ULiveLinkOBSDevice::HandleMessage(const FString& InMessage)
+{
+    TSharedPtr<FJsonObject> Root;
+    TSharedRef<TJsonReader<>> Reader = TJsonReaderFactory<>::Create(InMessage);
+    if (!FJsonSerializer::Deserialize(Reader, Root) || !Root.IsValid())
+    {
+        UE_LOGFMT(LogLiveLinkOBSDevice, Warning, "{DeviceName}: Failed to parse JSON: {Msg}",
+            GetSettings()->DisplayName, InMessage);
+        return;
+    }
+
+    int32 OpCode = 0;
+    if (!Root->TryGetNumberField(TEXT("op"), OpCode))
+    {
+        UE_LOGFMT(LogLiveLinkOBSDevice, Warning, "{DeviceName}: Message missing 'op' field.", GetSettings()->DisplayName);
+        return;
+    }
+
+    const TSharedPtr<FJsonObject>* DataObj = nullptr;
+    Root->TryGetObjectField(TEXT("d"), DataObj);
+    TSharedRef<FJsonObject> Data = DataObj ? DataObj->ToSharedRef() : MakeShared<FJsonObject>();
+
+    switch (OpCode)
+    {
+        case OBSWebSocketOpCode::Hello:
+            HandleHello(Data);
+            break;
+
+        case OBSWebSocketOpCode::Identified:
+            UE_LOGFMT(LogLiveLinkOBSDevice, Log, "{DeviceName}: Identified. Device ready.", GetSettings()->DisplayName);
+            bIdentified = true;
+            SendFilenameToOBS();
+            break;
+
+        case OBSWebSocketOpCode::Event:
+            HandleEvent(Data);
+            break;
+
+        case OBSWebSocketOpCode::Response:
+        {
+            // Fire-and-forget requests: just log failures.
+            const TSharedPtr<FJsonObject>* StatusObj = nullptr;
+            if (Data->TryGetObjectField(TEXT("requestStatus"), StatusObj))
+            {
+                bool bResult = false;
+                if ((*StatusObj)->TryGetBoolField(TEXT("result"), bResult) && !bResult)
+                {
+                    FString RequestType;
+                    int32   StatusCode = 0;
+                    Data->TryGetStringField(TEXT("requestType"), RequestType);
+                    (*StatusObj)->TryGetNumberField(TEXT("code"), StatusCode);
+                    UE_LOGFMT(LogLiveLinkOBSDevice, Warning,
+                        "{DeviceName}: Request '{RequestType}' failed with OBS status code {StatusCode}.",
+                        GetSettings()->DisplayName, RequestType, StatusCode);
+                }
+            }
+            break;
+        }
+
+        default:
+            break;
+    }
+}
+
+void ULiveLinkOBSDevice::HandleClosed(int32 StatusCode, const FString& Reason, bool bWasClean)
+{
+    UE_LOGFMT(LogLiveLinkOBSDevice, Log,
+        "{DeviceName}: WebSocket closed (code={StatusCode}, reason='{Reason}', clean={WasClean}).",
+        GetSettings()->DisplayName, StatusCode, Reason, bWasClean);
+
+    // OBS close code 4009 means the server rejected our authentication hash.
+    // Retrying with the same password would just loop forever, so treat this as
+    // an intentional disconnect and surface a clear error instead.
+    static constexpr int32 OBSCloseCode_AuthenticationFailed = 4009;
+    if (StatusCode == OBSCloseCode_AuthenticationFailed)
+    {
+        bIntentionalDisconnect = true; // always suppress reconnect on auth failure
+
+        if (bPasswordRequired)
+        {
+            // Blank password; bPasswordRequired was already set in HandleHello.
+            // Leave it set so the Warning state persists; don't escalate to Error.
+            UE_LOGFMT(LogLiveLinkOBSDevice, Warning,
+                "{DeviceName}: OBS requires authentication, no password configured.",
+                GetSettings()->DisplayName);
+        }
+        else
+        {
+            // Password was provided but OBS rejected it.
+            UE_LOGFMT(LogLiveLinkOBSDevice, Error,
+                "{DeviceName}: Authentication failed. Verify the password in device settings.",
+                GetSettings()->DisplayName);
+            bAuthFailed = true;
+        }
+    }
+
+    OnDisconnect();
+}
+
+void ULiveLinkOBSDevice::HandleError(const FString& Error)
+{
+    UE_LOGFMT(LogLiveLinkOBSDevice, Error, "{DeviceName}: WebSocket error: {Error}",
+        GetSettings()->DisplayName, Error);
+    OnDisconnect();
+}
+
+
+// -----------------------------------------------------------------------
+// OBS op handlers
+// -----------------------------------------------------------------------
+
+void ULiveLinkOBSDevice::HandleHello(const TSharedRef<FJsonObject>& InPayload)
+{
+    const TSharedPtr<FJsonObject>* AuthObj = nullptr;
+    if (InPayload->TryGetObjectField(TEXT("authentication"), AuthObj))
+    {
+        FString Challenge, Salt;
+        (*AuthObj)->TryGetStringField(TEXT("challenge"), Challenge);
+        (*AuthObj)->TryGetStringField(TEXT("salt"), Salt);
+
+        if (GetSettings()->Password.IsEmpty())
+        {
+            UE_LOGFMT(LogLiveLinkOBSDevice, Warning,
+                "{DeviceName}: OBS requires authentication but no password is configured.",
+                GetSettings()->DisplayName);
+            bPasswordRequired = true;
+        }
+
+        SendIdentify(ComputeAuthResponse(GetSettings()->Password, Challenge, Salt));
+    }
+    else
+    {
+        SendIdentify(FString());
+    }
+}
+
+void ULiveLinkOBSDevice::HandleEvent(const TSharedRef<FJsonObject>& InPayload)
+{
+    FString EventType;
+    if (!InPayload->TryGetStringField(TEXT("eventType"), EventType))
+    {
+        return;
+    }
+
+    if (EventType == TEXT("RecordStateChanged"))
+    {
+        const TSharedPtr<FJsonObject>* EventData = nullptr;
+        InPayload->TryGetObjectField(TEXT("eventData"), EventData);
+        if (!EventData)
+        {
+            return;
+        }
+
+        FString OutputState;
+        (*EventData)->TryGetStringField(TEXT("outputState"), OutputState);
+
+        const bool bNowRecording = (OutputState == TEXT("OBS_WEBSOCKET_OUTPUT_STARTED"));
+        if (bIsRecording != bNowRecording)
+        {
+            bIsRecording = bNowRecording;
+            UE_LOGFMT(LogLiveLinkOBSDevice, Log, "{DeviceName}: Recording state → {State}.",
+                GetSettings()->DisplayName, bIsRecording ? TEXT("Recording") : TEXT("Stopped"));
+        }
+    }
+}
+
+
+// -----------------------------------------------------------------------
+// Request helpers
+// -----------------------------------------------------------------------
+
+void ULiveLinkOBSDevice::SendIdentify(const FString& InAuthResponse)
+{
+    TSharedRef<FJsonObject> Data = MakeShared<FJsonObject>();
+    Data->SetNumberField(TEXT("rpcVersion"), 1);
+    if (!InAuthResponse.IsEmpty())
+    {
+        Data->SetStringField(TEXT("authentication"), InAuthResponse);
+    }
+    Data->SetNumberField(TEXT("eventSubscriptions"), 1 << 6); // Outputs category; includes RecordStateChanged
+
+    TSharedRef<FJsonObject> Root = MakeShared<FJsonObject>();
+    Root->SetNumberField(TEXT("op"), OBSWebSocketOpCode::Identify);
+    Root->SetObjectField(TEXT("d"), Data);
+
+    FString Out;
+    TSharedRef<TJsonWriter<>> Writer = TJsonWriterFactory<>::Create(&Out);
+    FJsonSerializer::Serialize(Root, Writer);
+    SendRaw(Out);
+}
+
+FString ULiveLinkOBSDevice::SendRequest(const FString& InRequestType, TSharedPtr<FJsonObject> InRequestData)
+{
+    if (!WebSocket || !WebSocket->IsConnected())
+    {
+        return FString();
+    }
+
+    const FString RequestId = FString::Printf(TEXT("llhub-%u"), ++RequestCounter);
+
+    TSharedRef<FJsonObject> Data = MakeShared<FJsonObject>();
+    Data->SetStringField(TEXT("requestType"), InRequestType);
+    Data->SetStringField(TEXT("requestId"),   RequestId);
+    if (InRequestData.IsValid())
+    {
+        Data->SetObjectField(TEXT("requestData"), InRequestData.ToSharedRef());
+    }
+
+    TSharedRef<FJsonObject> Root = MakeShared<FJsonObject>();
+    Root->SetNumberField(TEXT("op"), OBSWebSocketOpCode::Request);
+    Root->SetObjectField(TEXT("d"), Data);
+
+    FString Out;
+    TSharedRef<TJsonWriter<>> Writer = TJsonWriterFactory<>::Create(&Out);
+    FJsonSerializer::Serialize(Root, Writer);
+    SendRaw(Out);
+
+    return RequestId;
+}
+
+void ULiveLinkOBSDevice::SendRaw(const FString& InJson)
+{
+    if (WebSocket && WebSocket->IsConnected())
+    {
+        WebSocket->Send(InJson);
+    }
+}
+
+
+// -----------------------------------------------------------------------
+// Hub session callbacks
+// -----------------------------------------------------------------------
+
+void ULiveLinkOBSDevice::HandleSlateNameChanged(FStringView InSlateName)
+{
+    (void)InSlateName;
+    SendFilenameToOBS();
+}
+
+void ULiveLinkOBSDevice::HandleTakeNumberChanged(int32 InTakeNumber)
+{
+    (void)InTakeNumber;
+    SendFilenameToOBS();
+}
+
+void ULiveLinkOBSDevice::SendFilenameToOBS()
+{
+    if (!bIdentified)
+    {
+        return;
+    }
+
+    UNamingTokensEngineSubsystem* TokenSubsystem = GEngine->GetEngineSubsystem<UNamingTokensEngineSubsystem>();
+
+    FNamingTokenFilterArgs FilterArgs;
+    FilterArgs.AdditionalNamespacesToInclude.Add(TEXT("llh"));
+
+    FNamingTokenResultData Result = TokenSubsystem->EvaluateTokenText(
+        FText::FromString(GetSettings()->FilenameFormat),
+        FilterArgs
+    );
+
+    const FString Filename = Result.EvaluatedText.ToString();
+
+    UE_LOGFMT(LogLiveLinkOBSDevice, Log, "{DeviceName}: Setting OBS recording filename to '{Filename}'.",
+        GetSettings()->DisplayName, Filename);
+
+    // Update the "FilenameFormatting" key in the current OBS profile's [Output] section.
+    // OBS applies this when the next recording starts.
+    TSharedRef<FJsonObject> RequestData = MakeShared<FJsonObject>();
+    RequestData->SetStringField(TEXT("parameterCategory"), TEXT("Output"));
+    RequestData->SetStringField(TEXT("parameterName"),     TEXT("FilenameFormatting"));
+    RequestData->SetStringField(TEXT("parameterValue"),    Filename);
+
+    SendRequest(TEXT("SetProfileParameter"), RequestData);
+}
+
+
+// -----------------------------------------------------------------------
+// Auth
+// -----------------------------------------------------------------------
+
+FString ULiveLinkOBSDevice::ComputeAuthResponse(const FString& InPassword, const FString& InChallenge, const FString& InSalt)
+{
+    // OBS WebSocket v5 auth:
+    //   base64( SHA256( base64( SHA256( password + salt ) ) + challenge ) )
+
+    auto SHA256AndBase64 = [](const FString& InStr) -> FString
+    {
+        const FTCHARToUTF8 Utf8(*InStr);
+        uint8 Digest[32];
+#if LLOBS_WITH_OPENSSL
+        SHA256_CTX Ctx;
+        SHA256_Init(&Ctx);
+        SHA256_Update(&Ctx, Utf8.Get(), Utf8.Length());
+        SHA256_Final(Digest, &Ctx);
+#else
+        FMemory::Memzero(Digest);
+#endif
+        return FBase64::Encode(Digest, sizeof(Digest));
+    };
+
+    return SHA256AndBase64(SHA256AndBase64(InPassword + InSalt) + InChallenge);
+}
+
+
+#undef LOCTEXT_NAMESPACE
diff --git a/Engine/Plugins/VirtualProduction/LiveLinkOBSDevice/Source/LiveLinkOBSDevice/Private/Devices/LiveLinkOBSDevice.h b/Engine/Plugins/VirtualProduction/LiveLinkOBSDevice/Source/LiveLinkOBSDevice/Private/Devices/LiveLinkOBSDevice.h
new file mode 100644
--- /dev/null
+++ b/Engine/Plugins/VirtualProduction/LiveLinkOBSDevice/Source/LiveLinkOBSDevice/Private/Devices/LiveLinkOBSDevice.h
@@ -0,0 +1,205 @@
+// Copyright Epic Games, Inc. All Rights Reserved.
+
+#pragma once
+
+#include "Containers/Ticker.h"
+#include "LiveLinkDevice.h"
+#include "LiveLinkDeviceCapability_Connection.h"
+#include "LiveLinkDeviceCapability_Recording.h"
+#include "LiveLinkOBSDevice.generated.h"
+
+DECLARE_LOG_CATEGORY_EXTERN(LogLiveLinkOBSDevice, Log, All);
+
+class IWebSocket;
+
+// -----------------------------------------------------------------------
+// Settings
+// -----------------------------------------------------------------------
+
+/** Persisted settings for an OBS device instance. */
+UCLASS()
+class ULiveLinkOBSDeviceSettings : public ULiveLinkDeviceSettings
+{
+    GENERATED_BODY()
+
+public:
+    /** User-visible name shown in the device table. */
+    UPROPERTY(EditAnywhere, Category="OBS Device")
+    FString DisplayName = TEXT("OBS Studio");
+
+    /** Host running OBS WebSocket server (e.g. 127.0.0.1 or a remote IP). */
+    UPROPERTY(EditAnywhere, Category="OBS Device")
+    FString Host = TEXT("127.0.0.1");
+
+    /** Port configured in OBS Tools > WebSocket Server Settings (default 4455). */
+    UPROPERTY(EditAnywhere, Category="OBS Device")
+    int32 Port = 4455;
+
+    /** Password set in OBS WebSocket Server Settings. Leave empty if authentication is disabled. */
+    UPROPERTY(EditAnywhere, Category="OBS Device", meta=(PasswordField=true))
+    FString Password;
+
+    /**
+     * Format string for the OBS recording filename.
+     * Uses NamingTokens syntax; wrap token names in curly braces.
+     * LLH tokens (namespace "llh"): {slate}, {take}, {session}, {config}
+     * Global tokens: {project}, {user}, {yyyy}, {mm}, {Ddd}, etc.
+     */
+    UPROPERTY(EditAnywhere, Category="OBS Device", meta=(NamingTokens))
+    FString FilenameFormat = TEXT("{slate}-{take}");
+};
+
+
+// -----------------------------------------------------------------------
+// Device
+// -----------------------------------------------------------------------
+
+/**
+ * Live Link Hub device that controls OBS Studio recording via the OBS WebSocket v5 protocol.
+ *
+ * Supports:
+ *   - Start / Stop recording (driven by the Live Link Hub recording session)
+ *   - Setting the recording filename (driven by the Hub slate name)
+ */
+UCLASS(BlueprintType, meta=(DisplayName="OBS Studio"))
+class ULiveLinkOBSDevice : public ULiveLinkDevice
+    , public ILiveLinkDeviceCapability_Connection
+    , public ILiveLinkDeviceCapability_Recording
+{
+    GENERATED_BODY()
+
+public:
+    ULiveLinkOBSDeviceSettings* GetSettings()             { return GetDeviceSettings<ULiveLinkOBSDeviceSettings>(); }
+    const ULiveLinkOBSDeviceSettings* GetSettings() const { return GetDeviceSettings<ULiveLinkOBSDeviceSettings>(); }
+
+    //~ Begin ULiveLinkDevice interface
+    virtual TSubclassOf<ULiveLinkDeviceSettings> GetSettingsClass() const override { return ULiveLinkOBSDeviceSettings::StaticClass(); }
+    virtual FText         GetDisplayName()  const override;
+    virtual EDeviceHealth GetDeviceHealth() const override;
+    virtual FText         GetHealthText()   const override;
+    virtual void OnDeviceAdded()   override;
+    virtual void OnDeviceRemoved() override;
+    virtual void OnSettingChanged(const FPropertyChangedEvent& InPropertyChangedEvent) override;
+    //~ End ULiveLinkDevice interface
+
+    //~ Begin ILiveLinkDeviceCapability_Connection interface
+    virtual ELiveLinkDeviceConnectionStatus GetConnectionStatus_Implementation() const override;
+    virtual FString GetHardwareId_Implementation()                                    const override;
+    virtual bool    SetHardwareId_Implementation(const FString& InHardwareID)               override { return false; }
+    virtual bool    Connect_Implementation()                                                override;
+    virtual bool    Disconnect_Implementation()                                             override;
+protected:
+    virtual void    SetConnectionStatus(ELiveLinkDeviceConnectionStatus InStatus)           override;
+    //~ End ILiveLinkDeviceCapability_Connection interface
+
+public:
+    //~ Begin ILiveLinkDeviceCapability_Recording interface
+    virtual bool StartRecording_Implementation() override;
+    virtual bool StopRecording_Implementation()  override;
+    virtual bool IsRecording_Implementation() const override;
+    //~ End ILiveLinkDeviceCapability_Recording interface
+
+private:
+    // ----------------------------------------------------------------
+    // WebSocket lifecycle
+    // ----------------------------------------------------------------
+
+    /** Open the WebSocket connection to OBS. */
+    void Connect();
+
+    /** Close the WebSocket connection gracefully. */
+    void Disconnect();
+
+    /** Reset transient state (called on any disconnect path). */
+    void OnDisconnect();
+
+    // ----------------------------------------------------------------
+    // OBS WebSocket v5 protocol
+    // ----------------------------------------------------------------
+
+    /** Called when the WebSocket connects. OBS sends a Hello immediately after. */
+    void HandleConnected();
+
+    /** Called when the WebSocket receives a text message from OBS. */
+    void HandleMessage(const FString& InMessage);
+
+    /** Called when the WebSocket connection is lost. */
+    void HandleClosed(int32 StatusCode, const FString& Reason, bool bWasClean);
+
+    /** Called when the WebSocket encounters an error. */
+    void HandleError(const FString& Error);
+
+    // ----------------------------------------------------------------
+    // OBS WebSocket op handlers
+    // ----------------------------------------------------------------
+
+    /** Op 0 (Hello): OBS announces its version and optional auth challenge. */
+    void HandleHello(const TSharedRef<FJsonObject>& InPayload);
+
+    /** Op 5 (Event): OBS pushed an event (e.g. RecordStateChanged). */
+    void HandleEvent(const TSharedRef<FJsonObject>& InPayload);
+
+    // ----------------------------------------------------------------
+    // OBS WebSocket request helpers
+    // ----------------------------------------------------------------
+
+    /** Send an Identify (op 1) message. Pass the computed auth response, or empty string if auth is not required. */
+    void SendIdentify(const FString& InAuthResponse);
+
+    /** Send a request (op 6) to OBS. Returns the request ID used. */
+    FString SendRequest(const FString& InRequestType, TSharedPtr<FJsonObject> InRequestData = nullptr);
+
+    /** Build and send a raw JSON string over the WebSocket. */
+    void SendRaw(const FString& InJson);
+
+    // ----------------------------------------------------------------
+    // Live Link Hub session callbacks
+    // ----------------------------------------------------------------
+    void HandleSlateNameChanged(FStringView InSlateName);
+    void HandleTakeNumberChanged(int32 InTakeNumber);
+
+    /** Evaluate the FilenameFormat naming-tokens template and push the result to OBS as the recording filename. */
+    void SendFilenameToOBS();
+
+    // ----------------------------------------------------------------
+    // Auth helpers
+    // ----------------------------------------------------------------
+
+    /** Compute the OBS WebSocket v5 auth response string. */
+    static FString ComputeAuthResponse(const FString& InPassword, const FString& InChallenge, const FString& InSalt);
+
+    // ----------------------------------------------------------------
+    // State
+    // ----------------------------------------------------------------
+
+    ELiveLinkDeviceConnectionStatus ConnectionStatus = ELiveLinkDeviceConnectionStatus::Disconnected;
+
+    TSharedPtr<IWebSocket> WebSocket;
+
+    /** True once OBS has accepted our Identify message (op 2 Identified received). */
+    bool bIdentified = false;
+
+    /** Current recording state as reported by OBS events. */
+    bool bIsRecording = false;
+
+    /** Counter for generating unique request IDs. */
+    uint32 RequestCounter = 0;
+
+    /** True when Disconnect() was called deliberately; suppresses the auto-reconnect. */
+    bool bIntentionalDisconnect = false;
+
+    /** True when OBS rejected our authentication (close code 4009); suppresses reconnect and shows an error. */
+    bool bAuthFailed = false;
+
+    /** True when OBS requires authentication but the password field is blank. */
+    bool bPasswordRequired = false;
+
+    /** True when a connection setting changed while connected; triggers a reconnect once the current disconnect completes. */
+    bool bReconnectAfterSettingsChange = false;
+
+    /** Handle for the pending reconnect ticker. Valid only while a reconnect is scheduled. */
+    FTSTicker::FDelegateHandle ReconnectTickerHandle;
+
+    /** Seconds to wait before attempting to reconnect after an unexpected disconnect. */
+    static constexpr float ReconnectDelay = 5.f;
+};
diff --git a/Engine/Plugins/VirtualProduction/LiveLinkOBSDevice/Source/LiveLinkOBSDevice/Private/LiveLinkOBSDeviceModule.cpp b/Engine/Plugins/VirtualProduction/LiveLinkOBSDevice/Source/LiveLinkOBSDevice/Private/LiveLinkOBSDeviceModule.cpp
new file mode 100644
--- /dev/null
+++ b/Engine/Plugins/VirtualProduction/LiveLinkOBSDevice/Source/LiveLinkOBSDevice/Private/LiveLinkOBSDeviceModule.cpp
@@ -0,0 +1,38 @@
+// Copyright Epic Games, Inc. All Rights Reserved.
+
+#include "Customization/LiveLinkOBSDeviceSettingsCustomization.h"
+#include "Devices/LiveLinkOBSDevice.h"
+#include "Modules/ModuleManager.h"
+#include "PropertyEditorModule.h"
+
+
+class FLiveLinkOBSDeviceModule : public IModuleInterface
+{
+public:
+    virtual void StartupModule() override
+    {
+        FPropertyEditorModule& PropertyModule =
+            FModuleManager::LoadModuleChecked<FPropertyEditorModule>("PropertyEditor");
+
+        PropertyModule.RegisterCustomClassLayout(
+            ULiveLinkOBSDeviceSettings::StaticClass()->GetFName(),
+            FOnGetDetailCustomizationInstance::CreateStatic(
+                &FLiveLinkOBSDeviceSettingsCustomization::MakeInstance));
+
+        PropertyModule.NotifyCustomizationModuleChanged();
+    }
+
+    virtual void ShutdownModule() override
+    {
+        if (FModuleManager::Get().IsModuleLoaded("PropertyEditor"))
+        {
+            FPropertyEditorModule& PropertyModule =
+                FModuleManager::GetModuleChecked<FPropertyEditorModule>("PropertyEditor");
+
+            PropertyModule.UnregisterCustomClassLayout(
+                ULiveLinkOBSDeviceSettings::StaticClass()->GetFName());
+        }
+    }
+};
+
+IMPLEMENT_MODULE(FLiveLinkOBSDeviceModule, LiveLinkOBSDevice);
diff --git a/Engine/Source/Programs/LiveLinkHub/Source/LiveLinkHub.Target.cs b/Engine/Source/Programs/LiveLinkHub/Source/LiveLinkHub.Target.cs
--- a/Engine/Source/Programs/LiveLinkHub/Source/LiveLinkHub.Target.cs
+++ b/Engine/Source/Programs/LiveLinkHub/Source/LiveLinkHub.Target.cs
@@ -60,6 +60,7 @@ public class LiveLinkHubTarget : TargetRules
 			"LiveLinkUnrealDevice",
 			"LiveLinkKiProDevice",
 			"LiveLinkSoundDevice",
+			"LiveLinkOBSDevice",
 
 			"LiveLinkOpenVR",
 			"LiveLinkOpenTrackIO",
